Fix Xcode source code accessor name.

diff --git a/Engine/Plugins/Developer/XCodeSourceCodeAccess/Source/XCodeSourceCodeAccess/Private/XCodeSourceCodeAccessor.cpp b/Engine/Plugins/Developer/XCodeSourceCodeAccess/Source/XCodeSourceCodeAccess/Private/XCodeSourceCodeAccessor.cpp
--- a/Engine/Plugins/Developer/XCodeSourceCodeAccess/Source/XCodeSourceCodeAccess/Private/XCodeSourceCodeAccessor.cpp
+++ b/Engine/Plugins/Developer/XCodeSourceCodeAccess/Source/XCodeSourceCodeAccess/Private/XCodeSourceCodeAccessor.cpp
@@ -49,7 +49,7 @@ FName FXCodeSourceCodeAccessor::GetFName() const
 
 FText FXCodeSourceCodeAccessor::GetNameText() const
 {
-	return LOCTEXT("XCodeDisplayName", "Visual Studio");
+	return LOCTEXT("XCodeDisplayName", "Xcode");
 }
 
 FText FXCodeSourceCodeAccessor::GetDescriptionText() const
