Move GetSolutionFilepath() out of Core and into DesktopPlatform.

diff --git a/Engine/Plugins/Developer/VisualStudioSourceCodeAccess/Source/VisualStudioSourceCodeAccess/Private/VisualStudioSourceCodeAccessor.cpp b/Engine/Plugins/Developer/VisualStudioSourceCodeAccess/Source/VisualStudioSourceCodeAccess/Private/VisualStudioSourceCodeAccessor.cpp
--- a/Engine/Plugins/Developer/VisualStudioSourceCodeAccess/Source/VisualStudioSourceCodeAccess/Private/VisualStudioSourceCodeAccessor.cpp
+++ b/Engine/Plugins/Developer/VisualStudioSourceCodeAccess/Source/VisualStudioSourceCodeAccess/Private/VisualStudioSourceCodeAccessor.cpp
@@ -5,6 +5,7 @@
 #include "VisualStudioSourceCodeAccessModule.h"
 #include "ISourceCodeAccessModule.h"
 #include "ModuleManager.h"
+#include "DesktopPlatformModule.h"
 
 
 #include "AllowWindowsPlatformTypes.h"
@@ -906,8 +907,8 @@ FString FVisualStudioSourceCodeAccessor::GetSolutionPath() const
 	FScopeLock Lock(&CachedSolutionPathCriticalSection);
 	if(IsInGameThread())
 	{
-		FString SolutionPath = FModuleManager::Get().GetSolutionFilepath();
-		if(SolutionPath.Len() > 0)
+		FString SolutionPath;
+		if(FDesktopPlatformModule::Get()->GetSolutionPath(SolutionPath))
 		{
 			CachedSolutionPath = FPaths::ConvertRelativePathToFull(SolutionPath);
 		}
diff --git a/Engine/Plugins/Developer/VisualStudioSourceCodeAccess/Source/VisualStudioSourceCodeAccess/VisualStudioSourceCodeAccess.Build.cs b/Engine/Plugins/Developer/VisualStudioSourceCodeAccess/Source/VisualStudioSourceCodeAccess/VisualStudioSourceCodeAccess.Build.cs
--- a/Engine/Plugins/Developer/VisualStudioSourceCodeAccess/Source/VisualStudioSourceCodeAccess/VisualStudioSourceCodeAccess.Build.cs
+++ b/Engine/Plugins/Developer/VisualStudioSourceCodeAccess/Source/VisualStudioSourceCodeAccess/VisualStudioSourceCodeAccess.Build.cs
@@ -10,7 +10,8 @@ namespace UnrealBuildTool.Rules
 				new string[]
 				{
 					"Core",
-					"SourceCodeAccess"
+					"SourceCodeAccess",
+					"DesktopPlatform"
 				}
 			);
 
diff --git a/Engine/Plugins/Developer/XCodeSourceCodeAccess/Source/XCodeSourceCodeAccess/Private/XCodeSourceCodeAccessor.cpp b/Engine/Plugins/Developer/XCodeSourceCodeAccess/Source/XCodeSourceCodeAccess/Private/XCodeSourceCodeAccessor.cpp
--- a/Engine/Plugins/Developer/XCodeSourceCodeAccess/Source/XCodeSourceCodeAccess/Private/XCodeSourceCodeAccessor.cpp
+++ b/Engine/Plugins/Developer/XCodeSourceCodeAccess/Source/XCodeSourceCodeAccess/Private/XCodeSourceCodeAccessor.cpp
@@ -2,6 +2,7 @@
 
 #include "XCodeSourceCodeAccessPrivatePCH.h"
 #include "XCodeSourceCodeAccessor.h"
+#include "DesktopPlatformModule.h"
 
 #define LOCTEXT_NAMESPACE "XCodeSourceCodeAccessor"
 
@@ -59,11 +60,15 @@ FText FXCodeSourceCodeAccessor::GetDescriptionText() const
 
 bool FXCodeSourceCodeAccessor::OpenSolution()
 {
-	const FString FullPath = IFileManager::Get().ConvertToAbsolutePathForExternalAppForRead( *FModuleManager::Get().GetSolutionFilepath() );
-	if ( FPaths::FileExists( FullPath ) )
+	FString SolutionPath;
+	if(FDesktopPlatformModule::Get()->GetSolutionPath(SolutionPath))
 	{
-		FPlatformProcess::LaunchFileInDefaultExternalApplication( *FullPath );
-		return true;
+		const FString FullPath = IFileManager::Get().ConvertToAbsolutePathForExternalAppForRead( *SolutionPath );
+		if ( FPaths::FileExists( FullPath ) )
+		{
+			FPlatformProcess::LaunchFileInDefaultExternalApplication( *FullPath );
+			return true;
+		}
 	}
 	return false;
 }
@@ -82,9 +87,10 @@ bool FXCodeSourceCodeAccessor::OpenFileAtLine(const FString& FullPath, int32 Lin
 		ColumnNumber++;
 	}
 		 
-	if ( FModuleManager::Get().IsSolutionFilePresent() )
+	FString SolutionPath;
+	if(FDesktopPlatformModule::Get()->GetSolutionPath(SolutionPath))
 	{
-		FString ProjPath = FPaths::ConvertRelativePathToFull(FModuleManager::Get().GetSolutionFilepath());
+		FString ProjPath = FPaths::ConvertRelativePathToFull(*SolutionPath);
 		CFStringRef ProjPathString = FPlatformString::TCHARToCFString(*ProjPath);
 		NSString* ProjectPath = [(NSString*)ProjPathString stringByDeletingLastPathComponent];
 		[[NSWorkspace sharedWorkspace] openFile:ProjectPath withApplication:@"Xcode" andDeactivate:YES];
diff --git a/Engine/Plugins/Developer/XCodeSourceCodeAccess/Source/XCodeSourceCodeAccess/XCodeSourceCodeAccess.Build.cs b/Engine/Plugins/Developer/XCodeSourceCodeAccess/Source/XCodeSourceCodeAccess/XCodeSourceCodeAccess.Build.cs
--- a/Engine/Plugins/Developer/XCodeSourceCodeAccess/Source/XCodeSourceCodeAccess/XCodeSourceCodeAccess.Build.cs
+++ b/Engine/Plugins/Developer/XCodeSourceCodeAccess/Source/XCodeSourceCodeAccess/XCodeSourceCodeAccess.Build.cs
@@ -10,7 +10,8 @@ namespace UnrealBuildTool.Rules
 				new string[]
 				{
 					"Core",
-					"SourceCodeAccess"
+					"SourceCodeAccess",
+					"DesktopPlatform"
 				}
 			);
 		}
diff --git a/Engine/Plugins/Experimental/CodeView/Source/CodeView/CodeView.Build.cs b/Engine/Plugins/Experimental/CodeView/Source/CodeView/CodeView.Build.cs
--- a/Engine/Plugins/Experimental/CodeView/Source/CodeView/CodeView.Build.cs
+++ b/Engine/Plugins/Experimental/CodeView/Source/CodeView/CodeView.Build.cs
@@ -17,6 +17,7 @@ namespace UnrealBuildTool.Rules
 					"UnrealEd",
 					"EditorStyle",
 					"PropertyEditor",
+					"DesktopPlatform",
 				}
 			);
 
diff --git a/Engine/Plugins/Experimental/CodeView/Source/CodeView/Private/CodeViewPlugin.cpp b/Engine/Plugins/Experimental/CodeView/Source/CodeView/Private/CodeViewPlugin.cpp
--- a/Engine/Plugins/Experimental/CodeView/Source/CodeView/Private/CodeViewPlugin.cpp
+++ b/Engine/Plugins/Experimental/CodeView/Source/CodeView/Private/CodeViewPlugin.cpp
@@ -2,6 +2,7 @@
 
 #include "CodeViewPrivatePCH.h"
 #include "ModuleManager.h"
+#include "DesktopPlatformModule.h"
 
 class FCodeViewPlugin : public IModuleInterface
 {
@@ -20,7 +21,8 @@ public:
 
 	void AddCodeViewCategory(IDetailLayoutBuilder& DetailBuilder, const FGetSelectedActors& GetSelectedActors)
 	{
-		if (FModuleManager::Get().IsSolutionFilePresent())
+		FString SolutionPath;
+		if(FDesktopPlatformModule::Get()->GetSolutionPath(SolutionPath))
 		{
 			TSharedRef< CodeView::SCodeView > CodeViewWidget =
 				SNew( CodeView::SCodeView )
diff --git a/Engine/Source/Developer/DesktopPlatform/Private/DesktopPlatformBase.cpp b/Engine/Source/Developer/DesktopPlatform/Private/DesktopPlatformBase.cpp
--- a/Engine/Source/Developer/DesktopPlatform/Private/DesktopPlatformBase.cpp
+++ b/Engine/Source/Developer/DesktopPlatform/Private/DesktopPlatformBase.cpp
@@ -434,6 +434,40 @@ bool FDesktopPlatformBase::GenerateProjectFiles(const FString& RootDir, const FS
 	return bRes;
 }
 
+bool FDesktopPlatformBase::GetSolutionPath(FString& OutSolutionPath)
+{
+	// Get the platform-specific suffix for solution files
+#if PLATFORM_MAC
+	const TCHAR* Suffix = TEXT(".xcodeproj/project.pbxproj");
+#elif PLATFORM_LINUX
+	UE_LOG(LogModuleManager, Warning, TEXT("STUBBED: solution file path for Linux"));
+	const TCHAR* Suffix = TEXT("/stubbed/path/to.solution");
+#else
+	const TCHAR* Suffix = TEXT(".sln");
+#endif
+
+	// When using game specific uproject files, the solution is named after the game and in the uproject folder
+	if(FPaths::IsProjectFilePathSet())
+	{
+		FString SolutionPath = FPaths::GameDir() / FPaths::GetBaseFilename(FPaths::GetProjectFilePath()) + Suffix;
+		if(FPaths::FileExists(SolutionPath))
+		{
+			OutSolutionPath = SolutionPath;
+			return true;
+		}
+	}
+
+	// Otherwise, it is simply titled UE4.sln
+	FString DefaultSolutionPath = FPaths::RootDir() / FString(TEXT("UE4")) + Suffix;
+	if(FPaths::FileExists(DefaultSolutionPath))
+	{
+		OutSolutionPath = DefaultSolutionPath;
+		return true;
+	}
+
+	return false;
+}
+
 FString FDesktopPlatformBase::GetDefaultProjectCreationPath()
 {
 	// My Documents
diff --git a/Engine/Source/Developer/DesktopPlatform/Private/DesktopPlatformBase.h b/Engine/Source/Developer/DesktopPlatform/Private/DesktopPlatformBase.h
--- a/Engine/Source/Developer/DesktopPlatform/Private/DesktopPlatformBase.h
+++ b/Engine/Source/Developer/DesktopPlatform/Private/DesktopPlatformBase.h
@@ -39,6 +39,8 @@ public:
 	virtual bool CompileGameProject(const FString& RootDir, const FString& ProjectFileName, FFeedbackContext* Warn) override;
 	virtual bool GenerateProjectFiles(const FString& RootDir, const FString& ProjectFileName, FFeedbackContext* Warn) override;
 
+	virtual bool GetSolutionPath(FString& OutSolutionPath) override;
+
 	virtual bool EnumerateProjectsKnownByEngine(const FString &Identifier, bool bIncludeNativeProjects, TArray<FString> &OutProjectFileNames) override;
 	virtual FString GetDefaultProjectCreationPath() override;
 
diff --git a/Engine/Source/Developer/DesktopPlatform/Public/IDesktopPlatform.h b/Engine/Source/Developer/DesktopPlatform/Public/IDesktopPlatform.h
--- a/Engine/Source/Developer/DesktopPlatform/Public/IDesktopPlatform.h
+++ b/Engine/Source/Developer/DesktopPlatform/Public/IDesktopPlatform.h
@@ -340,6 +340,14 @@ public:
 	*/
 	virtual bool RunUnrealBuildTool(const FText& Description, const FString& RootDir, const FString& Arguments, FFeedbackContext* Warn) = 0;
 
+	/**
+	* Gets the path to the solution for the current project
+	*
+	* @param OutSolutionPath	Receives the string 
+	* @return True if a solution file exists and OutSolutionPath has been updated
+	*/
+	virtual bool GetSolutionPath(FString& OutSolutionPath) = 0;
+
 	/**
 	* Gets a feedback context which can display progress information using the native platform GUI.
 	*
diff --git a/Engine/Source/Developer/Toolbox/Private/Toolbox.cpp b/Engine/Source/Developer/Toolbox/Private/Toolbox.cpp
--- a/Engine/Source/Developer/Toolbox/Private/Toolbox.cpp
+++ b/Engine/Source/Developer/Toolbox/Private/Toolbox.cpp
@@ -6,13 +6,15 @@
 #include "ModuleUIInterface.h"
 #include "Editor/WorkspaceMenuStructure/Public/WorkspaceMenuStructureModule.h"
 #include "MainFrame.h"
+#include "DesktopPlatformModule.h"
 
 IMPLEMENT_MODULE( FToolboxModule, Toolbox );
 
 
 static bool CanShowModulesTab()
 {
-	return FModuleManager::Get().IsSolutionFilePresent();
+	FString SolutionPath;
+	return FDesktopPlatformModule::Get()->GetSolutionPath(SolutionPath);
 }
 
 class SDebugPanel : public SCompoundWidget
diff --git a/Engine/Source/Developer/Toolbox/Toolbox.Build.cs b/Engine/Source/Developer/Toolbox/Toolbox.Build.cs
--- a/Engine/Source/Developer/Toolbox/Toolbox.Build.cs
+++ b/Engine/Source/Developer/Toolbox/Toolbox.Build.cs
@@ -20,7 +20,8 @@ public class Toolbox : ModuleRules
                 "InputCore",
 				"Slate",
 				"SlateCore",
-                "EditorStyle"
+                "EditorStyle",
+				"DesktopPlatform"
 			}
 		);
 
diff --git a/Engine/Source/Editor/Documentation/Documentation.Build.cs b/Engine/Source/Editor/Documentation/Documentation.Build.cs
--- a/Engine/Source/Editor/Documentation/Documentation.Build.cs
+++ b/Engine/Source/Editor/Documentation/Documentation.Build.cs
@@ -41,7 +41,8 @@ namespace UnrealBuildTool.Rules
 					"Analytics",
 					"SourceCodeAccess",
 					"SourceControl",
-                    "ContentBrowser"
+                    "ContentBrowser",
+					"DesktopPlatform"
 				}
 			);
 
diff --git a/Engine/Source/Editor/Documentation/Private/UDNParser.cpp b/Engine/Source/Editor/Documentation/Private/UDNParser.cpp
--- a/Engine/Source/Editor/Documentation/Private/UDNParser.cpp
+++ b/Engine/Source/Editor/Documentation/Private/UDNParser.cpp
@@ -9,6 +9,7 @@
 #include "ContentBrowserModule.h"
 #include "AssetEditorManager.h"
 #include "SourceCodeNavigation.h"
+#include "DesktopPlatformModule.h"
 
 #define LOCTEXT_NAMESPACE "IntroTutorials"
 
@@ -1277,7 +1278,6 @@ bool FUDNParser::ParseCodeLink(FString &InternalLink)
 
 	ISourceCodeAccessModule& SourceCodeAccessModule = FModuleManager::LoadModuleChecked<ISourceCodeAccessModule>("SourceCodeAccess");
 	ISourceCodeAccessor& SourceCodeAccessor = SourceCodeAccessModule.GetAccessor();
-	FString SolutionPath;
 
 	// If we specified generic project specified as the project name try to replace the name with the name of this project
 	if (InternalLink.Contains(ProjectSpecifier) == true)
@@ -1302,8 +1302,8 @@ bool FUDNParser::ParseCodeLink(FString &InternalLink)
 
 	// Finally create the complete path - project name and all
 	int32 PathEndIndex;
-	SolutionPath = FModuleManager::Get().GetSolutionFilepath();
-	if (SolutionPath.FindLastChar(TEXT('/'), PathEndIndex) == true)
+	FString SolutionPath;
+	if( FDesktopPlatformModule::Get()->GetSolutionPath(SolutionPath) && SolutionPath.FindLastChar(TEXT('/'), PathEndIndex) == true)
 	{
 		SolutionPath = SolutionPath.LeftChop(SolutionPath.Len() - PathEndIndex - 1);
 		SolutionPath += Path;
diff --git a/Engine/Source/Editor/IntroTutorials/Private/STutorialEditableText.cpp b/Engine/Source/Editor/IntroTutorials/Private/STutorialEditableText.cpp
--- a/Engine/Source/Editor/IntroTutorials/Private/STutorialEditableText.cpp
+++ b/Engine/Source/Editor/IntroTutorials/Private/STutorialEditableText.cpp
@@ -6,6 +6,7 @@
 #include "ISourceCodeAccessModule.h"
 #include "ContentBrowserModule.h"
 #include "SColorPicker.h"
+#include "DesktopPlatformModule.h"
 
 
 #define LOCTEXT_NAMESPACE "STutorialEditableText"
@@ -71,7 +72,6 @@ static void ParseCodeLink(const FString &InternalLink)
 
 	ISourceCodeAccessModule& SourceCodeAccessModule = FModuleManager::LoadModuleChecked<ISourceCodeAccessModule>("SourceCodeAccess");
 	ISourceCodeAccessor& SourceCodeAccessor = SourceCodeAccessModule.GetAccessor();
-	FString SolutionPath;
 
 	// If we specified generic project specified as the project name try to replace the name with the name of this project
 	if (InternalLink.Contains(ProjectSpecifier) == true)
@@ -96,8 +96,8 @@ static void ParseCodeLink(const FString &InternalLink)
 
 	// Finally create the complete path - project name and all
 	int32 PathEndIndex;
-	SolutionPath = FModuleManager::Get().GetSolutionFilepath();
-	if (SolutionPath.FindLastChar(TEXT('/'), PathEndIndex) == true)
+	FString SolutionPath;
+	if( FDesktopPlatformModule::Get()->GetSolutionPath( SolutionPath ) && SolutionPath.FindLastChar(TEXT('/'), PathEndIndex) == true)
 	{
 		SolutionPath = SolutionPath.LeftChop(SolutionPath.Len() - PathEndIndex - 1);
 		SolutionPath += Path;
diff --git a/Engine/Source/Editor/Kismet/Kismet.Build.cs b/Engine/Source/Editor/Kismet/Kismet.Build.cs
--- a/Engine/Source/Editor/Kismet/Kismet.Build.cs
+++ b/Engine/Source/Editor/Kismet/Kismet.Build.cs
@@ -42,6 +42,7 @@ public class Kismet : ModuleRules
                 "Projects",
                 "JsonUtilities",
                 "DerivedDataCache",
+				"DesktopPlatform",
 			}
 			);
 
diff --git a/Engine/Source/Editor/Kismet/Private/SBlueprintEditorToolbar.cpp b/Engine/Source/Editor/Kismet/Private/SBlueprintEditorToolbar.cpp
--- a/Engine/Source/Editor/Kismet/Private/SBlueprintEditorToolbar.cpp
+++ b/Engine/Source/Editor/Kismet/Private/SBlueprintEditorToolbar.cpp
@@ -22,6 +22,7 @@
 #include "SBlueprintEditorSelectedDebugObjectWidget.h"
 #include "Engine/LevelScriptBlueprint.h"
 #include "Merge.h"
+#include "DesktopPlatformModule.h"
 
 #define LOCTEXT_NAMESPACE "KismetToolbar"
 
@@ -78,7 +79,8 @@ void FKismet2Menu::FillFileMenuBlueprintSection( FMenuBuilder& MenuBuilder, FBlu
 	MenuBuilder.EndSection();
 
 	// Only show the developer menu on machines with the solution (assuming they can build it)
-	if (FModuleManager::Get().IsSolutionFilePresent())
+	FString SolutionPath;
+	if(FDesktopPlatformModule::Get()->GetSolutionPath(SolutionPath))
 	{
 		MenuBuilder.BeginSection("FileDeveloper");
 		{
diff --git a/Engine/Source/Editor/MainFrame/Private/Frame/MainFrameActions.cpp b/Engine/Source/Editor/MainFrame/Private/Frame/MainFrameActions.cpp
--- a/Engine/Source/Editor/MainFrame/Private/Frame/MainFrameActions.cpp
+++ b/Engine/Source/Editor/MainFrame/Private/Frame/MainFrameActions.cpp
@@ -707,9 +707,17 @@ void FMainFrameActionCallbacks::OpenIDE()
 	{
 		if ( !FSourceCodeNavigation::OpenModuleSolution() )
 		{
-			const FString FullPath = IFileManager::Get().ConvertToAbsolutePathForExternalAppForRead( *FModuleManager::Get().GetSolutionFilepath() );
-			const FText Message = FText::Format( LOCTEXT( "OpenIDEFailed_MissingFile", "Could not open {0} for project {1}" ), FSourceCodeNavigation::GetSuggestedSourceCodeIDE(), FText::FromString( FullPath ) );
-			FMessageDialog::Open( EAppMsgType::Ok, Message );
+			FString SolutionPath;
+			if(FDesktopPlatformModule::Get()->GetSolutionPath(SolutionPath))
+			{
+				const FString FullPath = IFileManager::Get().ConvertToAbsolutePathForExternalAppForRead( *SolutionPath );
+				const FText Message = FText::Format( LOCTEXT( "OpenIDEFailed_MissingFile", "Could not open {0} for project {1}" ), FSourceCodeNavigation::GetSuggestedSourceCodeIDE(), FText::FromString( FullPath ) );
+				FMessageDialog::Open( EAppMsgType::Ok, Message );
+			}
+			else
+			{
+				FMessageDialog::Open( EAppMsgType::Ok, LOCTEXT("OpenIDEFailed_MissingSolution", "Couldn't find solution"));
+			}
 		}
 	}
 }
diff --git a/Engine/Source/Editor/MainFrame/Private/Menus/MainMenu.cpp b/Engine/Source/Editor/MainFrame/Private/Menus/MainMenu.cpp
--- a/Engine/Source/Editor/MainFrame/Private/Menus/MainMenu.cpp
+++ b/Engine/Source/Editor/MainFrame/Private/Menus/MainMenu.cpp
@@ -331,7 +331,8 @@ TSharedRef< SWidget > FMainMenu::MakeMainTabMenu( const TSharedPtr<FTabManager>&
 				);
 				*/
 
-				if (FPaths::FileExists(FModuleManager::Get().GetSolutionFilepath()))
+				FString SolutionPath;
+				if(FDesktopPlatformModule::Get()->GetSolutionPath(SolutionPath))
 				{
 					MenuBuilder.AddMenuEntry( FMainFrameCommands::Get().RefreshCodeProject,
 						NAME_None,
diff --git a/Engine/Source/Runtime/Core/Private/Modules/ModuleManager.cpp b/Engine/Source/Runtime/Core/Private/Modules/ModuleManager.cpp
--- a/Engine/Source/Runtime/Core/Private/Modules/ModuleManager.cpp
+++ b/Engine/Source/Runtime/Core/Private/Modules/ModuleManager.cpp
@@ -802,49 +802,6 @@ FString FModuleManager::GetModuleCompileMethod(FName InModuleName)
 }
 
 
-bool FModuleManager::IsSolutionFilePresent()
-{
-	return !GetSolutionFilepath().IsEmpty();
-}
-
-
-FString FModuleManager::GetSolutionFilepath()
-{
-#if PLATFORM_MAC
-	const TCHAR* Postfix = TEXT(".xcodeproj/project.pbxproj");
-#elif PLATFORM_LINUX
-	UE_LOG(LogModuleManager, Warning, TEXT("STUBBED: solution file path for Linux"));
-	const TCHAR* Postfix = TEXT("/stubbed/path/to.solution");
-#else
-	const TCHAR* Postfix = TEXT(".sln");
-#endif
-
-	FString SolutionFilepath;
-
-	if ( FPaths::IsProjectFilePathSet() )
-	{
-		// When using game specific uproject files, the solution is named after the game and in the uproject folder
-		SolutionFilepath = FPaths::GameDir() / FPaths::GetBaseFilename(FPaths::GetProjectFilePath()) + Postfix;
-		if ( !FPaths::FileExists(SolutionFilepath) )
-		{
-			SolutionFilepath = TEXT("");
-		}
-	}
-
-	if ( SolutionFilepath.IsEmpty() )
-	{
-		// Otherwise, it is simply titled UE4.sln
-		SolutionFilepath = FPaths::RootDir() / FString(TEXT("UE4")) + Postfix ;
-		if ( !FPaths::FileExists(SolutionFilepath) )
-		{
-			SolutionFilepath = TEXT("");
-		}
-	}
-
-	return SolutionFilepath;
-}
-
-
 bool FModuleManager::RecompileModule( const FName InModuleName, const bool bReloadAfterRecompile, FOutputDevice &Ar )
 {
 #if !IS_MONOLITHIC
diff --git a/Engine/Source/Runtime/Core/Public/Modules/ModuleManager.h b/Engine/Source/Runtime/Core/Public/Modules/ModuleManager.h
--- a/Engine/Source/Runtime/Core/Public/Modules/ModuleManager.h
+++ b/Engine/Source/Runtime/Core/Public/Modules/ModuleManager.h
@@ -410,21 +410,6 @@ public:
 	 */
 	void UnloadModulesAtShutdown( );
 
-	/**
-	 * Checks for the solution file using the hard-coded location on disk
-	 * Used to determine whether source code is potentially available for recompiles
-	 * 
-	 * @return	True if the solution file is found (source code MAY BE available)
-	 */
-	bool IsSolutionFilePresent();
-
-	/**
-	 * Returns the full path of the solution file
-	 * 
-	 * @return	SolutionFilepath
-	 */
-	FString GetSolutionFilepath();
-
 	/**
 	 * Tries to recompile the specified module.  If the module is loaded, it will first be unloaded (then reloaded after,
 	 * if the recompile was successful.)
