Implemented FXCodeSourceCodeAccessor::SaveAllOpenDocuments using Applescript.

diff --git a/Engine/Plugins/Developer/XCodeSourceCodeAccess/Source/XCodeSourceCodeAccess/Private/XCodeSourceCodeAccessor.cpp b/Engine/Plugins/Developer/XCodeSourceCodeAccess/Source/XCodeSourceCodeAccess/Private/XCodeSourceCodeAccessor.cpp
--- a/Engine/Plugins/Developer/XCodeSourceCodeAccess/Source/XCodeSourceCodeAccess/Private/XCodeSourceCodeAccessor.cpp
+++ b/Engine/Plugins/Developer/XCodeSourceCodeAccess/Source/XCodeSourceCodeAccess/Private/XCodeSourceCodeAccessor.cpp
@@ -38,6 +38,14 @@ static const char* OpenXCodeAtFileAndLineAppleScript =
 	"	end OpenXcodeAtFileAndLine\n"
 ;
 
+static const char* SaveAllXcodeDocuments =
+	"	on SaveAllXcodeDocuments()\n"
+	"		tell application \"Xcode\"\n"
+	"			save documents\n"
+	"		end tell\n"
+	"	end SaveAllXcodeDocuments\n"
+;
+
 bool FXCodeSourceCodeAccessor::CanAccessSourceCode() const
 {
 	return IFileManager::Get().DirectoryExists(TEXT("/Applications/Xcode.app"));
@@ -178,7 +186,29 @@ bool FXCodeSourceCodeAccessor::OpenSourceFiles(const TArray<FString>& AbsoluteSo
 
 bool FXCodeSourceCodeAccessor::SaveAllOpenDocuments() const
 {
-	return false;
+	bool ExecutionSucceeded = false;
+	
+	NSAppleScript* AppleScript = nil;
+	
+	NSString* AppleScriptString = [NSString stringWithCString:SaveAllXcodeDocuments encoding:NSUTF8StringEncoding];
+	AppleScript = [[NSAppleScript alloc] initWithSource:AppleScriptString];
+	
+	int PID = [[NSProcessInfo processInfo] processIdentifier];
+	NSAppleEventDescriptor* ThisApplication = [NSAppleEventDescriptor descriptorWithDescriptorType:typeKernelProcessID bytes:&PID length:sizeof(PID)];
+	
+	NSAppleEventDescriptor* ContainerEvent = [NSAppleEventDescriptor appleEventWithEventClass:'ascr' eventID:'psbr' targetDescriptor:ThisApplication returnID:kAutoGenerateReturnID transactionID:kAnyTransactionID];
+	
+	[ContainerEvent setParamDescriptor:[NSAppleEventDescriptor descriptorWithString:@"SaveAllXcodeDocuments"] forKeyword:'snam'];
+	
+	NSDictionary* ExecutionError = nil;
+	[AppleScript executeAppleEvent:ContainerEvent error:&ExecutionError];
+	if(ExecutionError == nil)
+	{
+		ExecutionSucceeded = true;
+	}
+	
+	[AppleScript release];
+	return ExecutionSucceeded;
 }
 
 void FXCodeSourceCodeAccessor::Tick(const float DeltaTime)
