Harden the Xcode source accessor Applescript so that it can reactivate Xcode if the user switches to something else between Apple events as that causes the editor to appear to hang - it is waiting for Xcode's "Open Quickly" window to appear.

diff --git a/Engine/Plugins/Developer/XCodeSourceCodeAccess/Source/XCodeSourceCodeAccess/Private/XCodeSourceCodeAccessor.cpp b/Engine/Plugins/Developer/XCodeSourceCodeAccess/Source/XCodeSourceCodeAccess/Private/XCodeSourceCodeAccessor.cpp
--- a/Engine/Plugins/Developer/XCodeSourceCodeAccess/Source/XCodeSourceCodeAccess/Private/XCodeSourceCodeAccessor.cpp
+++ b/Engine/Plugins/Developer/XCodeSourceCodeAccess/Source/XCodeSourceCodeAccess/Private/XCodeSourceCodeAccessor.cpp
@@ -6,36 +6,55 @@
 
 #define LOCTEXT_NAMESPACE "XCodeSourceCodeAccessor"
 
+DEFINE_LOG_CATEGORY_STATIC(LogXcodeAccessor, Log, All);
+
 /** Applescript we use to open XCode */
 static const char* OpenXCodeAtFileAndLineAppleScript =
-	"	on OpenXcodeAtFileAndLine(filepath, linenumber)\n"
-	"		set theOffset to offset of \"/\" in filepath\n"
-	"		tell application \"Xcode\"\n"
-	"			activate\n"
-	"			if theOffset is 1 then\n"
-	"				open filepath\n"
-	"			end if\n"
-	"			tell application \"System Events\"\n"
-	"				tell process \"Xcode\"\n"
-	"					if theOffset is not 1 then\n"
-	"						keystroke \"o\" using {command down, shift down}\n"
-	"						repeat until window \"Open Quickly\" exists\n"
-	"						end repeat\n"
-	"						click text field 1 of window \"Open Quickly\"\n"
-	"						set value of text field 1 of window \"Open Quickly\" to filepath\n"
-	"						keystroke return\n"
-	"					end if\n"
-	"					keystroke \"l\" using command down\n"
+	"on OpenXcodeAtFileAndLine(filepath, linenumber)\n"
+	"	set theOffset to offset of \"/\" in filepath\n"
+	"	tell application \"Xcode\"\n"
+	"		activate\n"
+	"		if theOffset is 1 then\n"
+	"			open filepath\n"
+	"		end if\n"
+	"		tell application \"System Events\"\n"
+	"			tell process \"Xcode\"\n"
+	"				\n"
+	"				if theOffset is not 1 then\n"
 	"					repeat until window \"Open Quickly\" exists\n"
+	"						tell application \"Xcode\"\n"
+	"							if application \"Xcode\" is not frontmost then\n"
+	"								activate\n"
+	"							end if\n"
+	"						end tell\n"
+	"						if application \"Xcode\" is frontmost then\n"
+	"							keystroke \"o\" using {command down, shift down}\n"
+	"						end if\n"
 	"					end repeat\n"
 	"					click text field 1 of window \"Open Quickly\"\n"
-	"					set value of text field 1 of window \"Open Quickly\" to linenumber\n"
+	"					set value of text field 1 of window \"Open Quickly\" to filepath\n"
 	"					keystroke return\n"
-	"					keystroke return\n"
-	"				end tell\n"
+	"				end if\n"
+	"				\n"
+	"				repeat until window \"Open Quickly\" exists\n"
+	"					tell application \"Xcode\"\n"
+	"						if application \"Xcode\" is not frontmost then\n"
+	"							activate\n"
+	"						end if\n"
+	"					end tell\n"
+	"					if application \"Xcode\" is frontmost then\n"
+	"						keystroke \"l\" using command down\n"
+	"					end if\n"
+	"				end repeat\n"
+	"				\n"
+	"				click text field 1 of window \"Open Quickly\"\n"
+	"				set value of text field 1 of window \"Open Quickly\" to linenumber\n"
+	"				keystroke return\n"
+	"				keystroke return\n"
 	"			end tell\n"
 	"		end tell\n"
-	"	end OpenXcodeAtFileAndLine\n"
+	"	end tell\n"
+	"end OpenXcodeAtFileAndLine\n"
 ;
 
 static const char* SaveAllXcodeDocuments =
@@ -162,6 +181,10 @@ bool FXCodeSourceCodeAccessor::OpenFileAtLine(const FString& FullPath, int32 Lin
 	{
 		ExecutionSucceeded = true;
 	}
+	else
+	{
+		UE_LOG(LogXcodeAccessor, Error, TEXT("%s"), *FString([ExecutionError description]));
+	}
 		
 	[AppleScript release];
 		
@@ -211,6 +234,10 @@ bool FXCodeSourceCodeAccessor::SaveAllOpenDocuments() const
 	{
 		ExecutionSucceeded = true;
 	}
+	else
+	{
+		UE_LOG(LogXcodeAccessor, Error, TEXT("%s"), *FString([ExecutionError description]));
+	}
 	
 	[AppleScript release];
 	return ExecutionSucceeded;
