When opening at a given file & line in Xcode make sure that the project file exists before trying to execute the Applescript, or it won't work.

diff --git a/Engine/Plugins/Developer/XCodeSourceCodeAccess/Source/XCodeSourceCodeAccess/Private/XCodeSourceCodeAccessor.cpp b/Engine/Plugins/Developer/XCodeSourceCodeAccess/Source/XCodeSourceCodeAccess/Private/XCodeSourceCodeAccessor.cpp
--- a/Engine/Plugins/Developer/XCodeSourceCodeAccess/Source/XCodeSourceCodeAccess/Private/XCodeSourceCodeAccessor.cpp
+++ b/Engine/Plugins/Developer/XCodeSourceCodeAccess/Source/XCodeSourceCodeAccess/Private/XCodeSourceCodeAccessor.cpp
@@ -21,14 +21,16 @@ static const char* OpenXCodeAtFileAndLineAppleScript =
 	"			tell process \"Xcode\"\n"
 	"				\n"
 	"				if theOffset is not 1 then\n"
+	"					set bActivated to false\n"
 	"					repeat until window \"Open Quickly\" exists\n"
 	"						tell application \"Xcode\"\n"
 	"							if application \"Xcode\" is not frontmost then\n"
 	"								activate\n"
 	"							end if\n"
 	"						end tell\n"
-	"						if application \"Xcode\" is frontmost then\n"
+	"						if application \"Xcode\" is frontmost and bActivated is false then\n"
 	"							keystroke \"o\" using {command down, shift down}\n"
+	"							set bActivated to true\n"
 	"						end if\n"
 	"					end repeat\n"
 	"					click text field 1 of window \"Open Quickly\"\n"
@@ -36,14 +38,16 @@ static const char* OpenXCodeAtFileAndLineAppleScript =
 	"					keystroke return\n"
 	"				end if\n"
 	"				\n"
+	"				set bActivated to false\n"
 	"				repeat until window \"Open Quickly\" exists\n"
 	"					tell application \"Xcode\"\n"
 	"						if application \"Xcode\" is not frontmost then\n"
 	"							activate\n"
 	"						end if\n"
 	"					end tell\n"
-	"					if application \"Xcode\" is frontmost then\n"
+	"					if application \"Xcode\" is frontmost and bActivated is false then\n"
 	"						keystroke \"l\" using command down\n"
+	"						set bActivated to true\n"
 	"					end if\n"
 	"				end repeat\n"
 	"				\n"
@@ -114,87 +118,90 @@ bool FXCodeSourceCodeAccessor::OpenFileAtLine(const FString& FullPath, int32 Lin
 		ColumnNumber++;
 	}
 		 
+	bool ExecutionSucceeded = false;
+	
 	FString SolutionPath;
 	if(FDesktopPlatformModule::Get()->GetSolutionPath(SolutionPath))
 	{
-		FString ProjPath = FPaths::ConvertRelativePathToFull(*SolutionPath);
-		CFStringRef ProjPathString = FPlatformString::TCHARToCFString(*ProjPath);
-		NSString* ProjectPath = [(NSString*)ProjPathString stringByDeletingLastPathComponent];
-		[[NSWorkspace sharedWorkspace] openFile:ProjectPath withApplication:@"Xcode" andDeactivate:YES];
-	}
-		
-	bool ExecutionSucceeded = false;
-		
-	NSAppleScript* AppleScript = nil;
-		
-	NSString* AppleScriptString = [NSString stringWithCString:OpenXCodeAtFileAndLineAppleScript encoding:NSUTF8StringEncoding];
-	AppleScript = [[NSAppleScript alloc] initWithSource:AppleScriptString];
-		
-	int PID = [[NSProcessInfo processInfo] processIdentifier];
-	NSAppleEventDescriptor* ThisApplication = [NSAppleEventDescriptor descriptorWithDescriptorType:typeKernelProcessID bytes:&PID length:sizeof(PID)];
+		const FString ProjPath = IFileManager::Get().ConvertToAbsolutePathForExternalAppForRead( *SolutionPath );
+		if ( FPaths::FileExists( ProjPath ) )
+		{
+			NSString* ProjectPath = [ProjPath.GetNSString() stringByDeletingLastPathComponent];
+			[[NSWorkspace sharedWorkspace] openFile:ProjectPath withApplication:@"Xcode" andDeactivate:YES];
 			
-	NSAppleEventDescriptor* ContainerEvent = [NSAppleEventDescriptor appleEventWithEventClass:'ascr' eventID:'psbr' targetDescriptor:ThisApplication returnID:kAutoGenerateReturnID transactionID:kAnyTransactionID];
+			NSAppleScript* AppleScript = nil;
 			
-	[ContainerEvent setParamDescriptor:[NSAppleEventDescriptor descriptorWithString:@"OpenXcodeAtFileAndLine"] forKeyword:'snam'];
+			NSString* AppleScriptString = [NSString stringWithCString:OpenXCodeAtFileAndLineAppleScript encoding:NSUTF8StringEncoding];
+			AppleScript = [[NSAppleScript alloc] initWithSource:AppleScriptString];
 			
-	{
-		NSAppleEventDescriptor* Arguments = [[NSAppleEventDescriptor alloc] initListDescriptor];
-
-		CFStringRef FileString = FPlatformString::TCHARToCFString(*FullPath);
-		NSString* Path = (NSString*)FileString;
+			int PID = [[NSProcessInfo processInfo] processIdentifier];
+			NSAppleEventDescriptor* ThisApplication = [NSAppleEventDescriptor descriptorWithDescriptorType:typeKernelProcessID bytes:&PID length:sizeof(PID)];
+			
+			NSAppleEventDescriptor* ContainerEvent = [NSAppleEventDescriptor appleEventWithEventClass:'ascr' eventID:'psbr' targetDescriptor:ThisApplication returnID:kAutoGenerateReturnID transactionID:kAnyTransactionID];
+			
+			[ContainerEvent setParamDescriptor:[NSAppleEventDescriptor descriptorWithString:@"OpenXcodeAtFileAndLine"] forKeyword:'snam'];
+			
+			{
+				NSAppleEventDescriptor* Arguments = [[NSAppleEventDescriptor alloc] initListDescriptor];
 				
-		if([Path isAbsolutePath] == NO)
-		{
-			NSString* CurDir = [[NSFileManager defaultManager] currentDirectoryPath];
-			NSString* ResolvedPath = [[NSString stringWithFormat:@"%@/%@", CurDir, Path] stringByResolvingSymlinksInPath];
-			if([[NSFileManager defaultManager] fileExistsAtPath:ResolvedPath])
+				CFStringRef FileString = FPlatformString::TCHARToCFString(*FullPath);
+				NSString* Path = (NSString*)FileString;
+				
+				if([Path isAbsolutePath] == NO)
+				{
+					NSString* CurDir = [[NSFileManager defaultManager] currentDirectoryPath];
+					NSString* ResolvedPath = [[NSString stringWithFormat:@"%@/%@", CurDir, Path] stringByResolvingSymlinksInPath];
+					if([[NSFileManager defaultManager] fileExistsAtPath:ResolvedPath])
+					{
+						Path = ResolvedPath;
+					}
+					else // If it doesn't exist, supply only the filename, we'll use Open Quickly to try and find it from Xcode
+					{
+						Path = [Path lastPathComponent];
+					}
+				}
+				
+				[Arguments insertDescriptor:[NSAppleEventDescriptor descriptorWithString:Path] atIndex:([Arguments numberOfItems] + 1)];
+				CFRelease(FileString);
+				
+				CFStringRef LineString = FPlatformString::TCHARToCFString(*FString::FromInt(LineNumber));
+				if(LineString)
+				{
+					[Arguments insertDescriptor:[NSAppleEventDescriptor descriptorWithString:(NSString*)LineString] atIndex:([Arguments numberOfItems] + 1)];
+					CFRelease(LineString);
+				}
+				else
+				{
+					[Arguments insertDescriptor:[NSAppleEventDescriptor descriptorWithString:@"1"] atIndex:([Arguments numberOfItems] + 1)];
+				}
+				
+				[ContainerEvent setParamDescriptor:Arguments forKeyword:keyDirectObject];
+				[Arguments release];
+			}
+			
+			NSDictionary* ExecutionError = nil;
+			[AppleScript executeAppleEvent:ContainerEvent error:&ExecutionError];
+			if(ExecutionError == nil)
 			{
-				Path = ResolvedPath;
+				ExecutionSucceeded = true;
 			}
-			else // If it doesn't exist, supply only the filename, we'll use Open Quickly to try and find it from Xcode
+			else
 			{
-				Path = [Path lastPathComponent];
+				UE_LOG(LogXcodeAccessor, Error, TEXT("%s"), *FString([ExecutionError description]));
 			}
-		}
-				
-		[Arguments insertDescriptor:[NSAppleEventDescriptor descriptorWithString:Path] atIndex:([Arguments numberOfItems] + 1)];
-		CFRelease(FileString);
-				
-		CFStringRef LineString = FPlatformString::TCHARToCFString(*FString::FromInt(LineNumber));
-		if(LineString)
-		{
-			[Arguments insertDescriptor:[NSAppleEventDescriptor descriptorWithString:(NSString*)LineString] atIndex:([Arguments numberOfItems] + 1)];
-			CFRelease(LineString);
-		}
-		else
-		{
-			[Arguments insertDescriptor:[NSAppleEventDescriptor descriptorWithString:@"1"] atIndex:([Arguments numberOfItems] + 1)];
-		}
-				
-		[ContainerEvent setParamDescriptor:Arguments forKeyword:keyDirectObject];
-		[Arguments release];
-	}
 			
-	NSDictionary* ExecutionError = nil;
-	[AppleScript executeAppleEvent:ContainerEvent error:&ExecutionError];
-	if(ExecutionError == nil)
-	{
-		ExecutionSucceeded = true;
-	}
-	else
-	{
-		UE_LOG(LogXcodeAccessor, Error, TEXT("%s"), *FString([ExecutionError description]));
+			[AppleScript release];
+			
+			// Fallback to trivial implementation when something goes wrong (like not having permission for UI scripting)
+			if(ExecutionSucceeded == false)
+			{
+				FPlatformProcess::LaunchFileInDefaultExternalApplication(*FullPath);
+				ExecutionSucceeded = true;
+			}
+		}
 	}
 		
-	[AppleScript release];
-		
-	// Fallback to trivial implementation when something goes wrong (like not having permission for UI scripting)
-	if(ExecutionSucceeded == false)
-	{
-		FPlatformProcess::LaunchFileInDefaultExternalApplication(*FullPath);
-	}
-
-	return true;
+	return ExecutionSucceeded;
 }
 
 bool FXCodeSourceCodeAccessor::OpenSourceFiles(const TArray<FString>& AbsoluteSourcePaths)
