Fix for KeyUp triggers us hitting the disk duing PIE and killing performance.  Stop updating menu when PIE viewport has focus.  Tested without this in place where things like CMD+N don't function anyway with PIE viewport focused, so it keeps that functionality.  Tested PIE selecgtedviewport and new editor window.

Also use the CachedSolutionPath in XcodeCodeAccessor as an actual cache rather then computing each time.

diff --git a/Engine/Plugins/Developer/XCodeSourceCodeAccess/Source/XCodeSourceCodeAccess/Private/XCodeSourceCodeAccessor.cpp b/Engine/Plugins/Developer/XCodeSourceCodeAccess/Source/XCodeSourceCodeAccess/Private/XCodeSourceCodeAccessor.cpp
--- a/Engine/Plugins/Developer/XCodeSourceCodeAccess/Source/XCodeSourceCodeAccess/Private/XCodeSourceCodeAccessor.cpp
+++ b/Engine/Plugins/Developer/XCodeSourceCodeAccess/Source/XCodeSourceCodeAccess/Private/XCodeSourceCodeAccessor.cpp
@@ -82,22 +82,26 @@ void FXCodeSourceCodeAccessor::Startup()
 
 void FXCodeSourceCodeAccessor::Shutdown()
 {
+	CachedSolutionPath.Empty();
 }
 
 FString FXCodeSourceCodeAccessor::GetSolutionPath() const
 {
 	if(IsInGameThread())
 	{
-		CachedSolutionPath = FPaths::ProjectDir();
-		
-		if (!FUProjectDictionary(FPaths::RootDir()).IsForeignProject(CachedSolutionPath))
-		{
-			CachedSolutionPath = FPaths::Combine(FPaths::RootDir(), + TEXT("UE4.xcworkspace/contents.xcworkspacedata"));
-		}
-		else
+		if(CachedSolutionPath.Len() <= 0)
 		{
-            FString BaseName = FApp::HasProjectName() ? FApp::GetProjectName() : FPaths::GetBaseFilename(CachedSolutionPath);
-			CachedSolutionPath = FPaths::Combine(CachedSolutionPath, BaseName + TEXT(".xcworkspace/contents.xcworkspacedata"));
+			CachedSolutionPath = FPaths::ProjectDir();
+			
+			if (!FUProjectDictionary(FPaths::RootDir()).IsForeignProject(CachedSolutionPath))
+			{
+				CachedSolutionPath = FPaths::Combine(FPaths::RootDir(), + TEXT("UE4.xcworkspace/contents.xcworkspacedata"));
+			}
+			else
+			{
+				FString BaseName = FApp::HasProjectName() ? FApp::GetProjectName() : FPaths::GetBaseFilename(CachedSolutionPath);
+				CachedSolutionPath = FPaths::Combine(CachedSolutionPath, BaseName + TEXT(".xcworkspace/contents.xcworkspacedata"));
+			}
 		}
 	}
 	return CachedSolutionPath;
diff --git a/Engine/Source/Runtime/Slate/Private/Framework/MultiBox/Mac/MacMenu.cpp b/Engine/Source/Runtime/Slate/Private/Framework/MultiBox/Mac/MacMenu.cpp
--- a/Engine/Source/Runtime/Slate/Private/Framework/MultiBox/Mac/MacMenu.cpp
+++ b/Engine/Source/Runtime/Slate/Private/Framework/MultiBox/Mac/MacMenu.cpp
@@ -8,6 +8,7 @@
 #include "HAL/PlatformApplicationMisc.h"
 #include "Framework/Commands/Commands.h"
 #include "Misc/App.h"
+#include "Widgets/SViewport.h"
 
 struct FMacMenuItemState
 {
@@ -700,7 +701,19 @@ void FSlateMacMenu::UpdateCachedState()
 		}
 		MacApplication->GetWindowsArrayMutex().Unlock();
     }
-	
+    
+    // If PIE Viewport has focus don't update
+    if(GIsEditor && FSlateApplication::IsInitialized())
+    {
+    	TSharedPtr<SViewport> ViewPort = FSlateApplication::Get().GetGameViewport();
+		if(ViewPort.IsValid())
+		{
+			if(ViewPort->HasKeyboardFocus())
+			{
+				bShouldUpdate = false;
+			}
+		}
+	}
 
 	if (bShouldUpdate)
 	{
