UE5EA - Update UE4 reference

diff --git a/Engine/Plugins/Developer/XCodeSourceCodeAccess/Source/XCodeSourceCodeAccess/Private/XCodeSourceCodeAccessor.cpp b/Engine/Plugins/Developer/XCodeSourceCodeAccess/Source/XCodeSourceCodeAccess/Private/XCodeSourceCodeAccessor.cpp
--- a/Engine/Plugins/Developer/XCodeSourceCodeAccess/Source/XCodeSourceCodeAccess/Private/XCodeSourceCodeAccessor.cpp
+++ b/Engine/Plugins/Developer/XCodeSourceCodeAccess/Source/XCodeSourceCodeAccess/Private/XCodeSourceCodeAccessor.cpp
@@ -95,7 +95,7 @@ FString FXCodeSourceCodeAccessor::GetSolutionPath() const
 			
 			if (!FUProjectDictionary::GetDefault().IsForeignProject(CachedSolutionPath))
 			{
-				CachedSolutionPath = FPaths::Combine(FPaths::RootDir(), + TEXT("UE4.xcworkspace/contents.xcworkspacedata"));
+				CachedSolutionPath = FPaths::Combine(FPaths::RootDir(), + TEXT("UE5.xcworkspace/contents.xcworkspacedata"));
 			}
 			else
 			{
