Fix Tools -> Open Xcode not working with modern Xcode

diff --git a/Engine/Plugins/Developer/XCodeSourceCodeAccess/Source/XCodeSourceCodeAccess/Private/XCodeSourceCodeAccessor.cpp b/Engine/Plugins/Developer/XCodeSourceCodeAccess/Source/XCodeSourceCodeAccess/Private/XCodeSourceCodeAccessor.cpp
--- a/Engine/Plugins/Developer/XCodeSourceCodeAccess/Source/XCodeSourceCodeAccess/Private/XCodeSourceCodeAccessor.cpp
+++ b/Engine/Plugins/Developer/XCodeSourceCodeAccess/Source/XCodeSourceCodeAccess/Private/XCodeSourceCodeAccessor.cpp
@@ -95,13 +95,21 @@ FString FXCodeSourceCodeAccessor::GetSolutionPath() const
 			
 			if (!FUProjectDictionary::GetDefault().IsForeignProject(CachedSolutionPath))
 			{
-				CachedSolutionPath = FPaths::Combine(FPaths::RootDir(), + TEXT("UE5.xcworkspace/contents.xcworkspacedata"));
+				CachedSolutionPath = FPaths::Combine(FPaths::RootDir(), + TEXT("UE5 (Mac).xcworkspace/contents.xcworkspacedata"));
 			}
 			else
 			{
 				FString BaseName = FApp::HasProjectName() ? FApp::GetProjectName() : FPaths::GetBaseFilename(CachedSolutionPath);
-				CachedSolutionPath = FPaths::Combine(CachedSolutionPath, BaseName + TEXT(".xcworkspace/contents.xcworkspacedata"));
+				CachedSolutionPath = FPaths::Combine(CachedSolutionPath, BaseName + TEXT(" (Mac).xcworkspace/contents.xcworkspacedata"));
 			}
+            
+            // If modern xcode project doesn't exist, use legacy name instead
+            const FString FullPath = IFileManager::Get().ConvertToAbsolutePathForExternalAppForRead( *CachedSolutionPath );
+            if(!FPaths::FileExists( FullPath ))
+            {
+                CachedSolutionPath.RemoveFromEnd(TEXT(" (Mac).xcworkspace/contents.xcworkspacedata"));
+                CachedSolutionPath += TEXT(".xcworkspace/contents.xcworkspacedata");
+            }
 		}
 	}
 	return CachedSolutionPath;
